[AI Assistant] Gameplay tags toolset

#Synced-CL 52269267

diff --git a/Engine/Plugins/Experimental/Toolsets/GameplayTagsToolset/GameplayTagsToolset.uplugin b/Engine/Plugins/Experimental/Toolsets/GameplayTagsToolset/GameplayTagsToolset.uplugin
new file mode 100644
--- /dev/null
+++ b/Engine/Plugins/Experimental/Toolsets/GameplayTagsToolset/GameplayTagsToolset.uplugin
@@ -0,0 +1,30 @@
+{
+	"FileVersion": 3,
+	"Version": 1,
+	"VersionName": "1.0",
+	"FriendlyName": "Gameplay Tags Toolset",
+	"Description": "Toolset for reading and managing gameplay tags via the AI Toolset Registry.",
+	"Category": "Other",
+	"EditorOnly": true,
+	"IsExperimentalVersion": true,
+	"NoRedist": true,
+	"EnabledByDefault": false,
+	"Modules": [
+		{
+			"Name": "GameplayTagsToolset",
+			"Type": "Editor",
+			"LoadingPhase": "PostEngineInit",
+			"TargetAllowList": ["Editor"]
+		}
+	],
+	"Plugins": [
+		{
+			"Name": "ToolsetRegistry",
+			"Enabled": true
+		},
+		{
+			"Name": "GameplayTagsEditor",
+			"Enabled": true
+		}
+	]
+}
diff --git a/Engine/Plugins/Experimental/Toolsets/GameplayTagsToolset/Source/GameplayTagsToolset/GameplayTagsToolset.Build.cs b/Engine/Plugins/Experimental/Toolsets/GameplayTagsToolset/Source/GameplayTagsToolset/GameplayTagsToolset.Build.cs
new file mode 100644
--- /dev/null
+++ b/Engine/Plugins/Experimental/Toolsets/GameplayTagsToolset/Source/GameplayTagsToolset/GameplayTagsToolset.Build.cs
@@ -0,0 +1,26 @@
+// Copyright Epic Games, Inc. All Rights Reserved.
+
+using UnrealBuildTool;
+
+public class GameplayTagsToolset : ModuleRules
+{
+	public GameplayTagsToolset(ReadOnlyTargetRules Target) : base(Target)
+	{
+		PCHUsage = PCHUsageMode.UseExplicitOrSharedPCHs;
+		IWYUSupport = IWYUSupport.Full;
+		bUseUnity = true;
+
+		PublicDependencyModuleNames.Add("Core");
+
+		PrivateDependencyModuleNames.AddRange(new string[]
+		{
+			"CoreUObject",
+			"Engine",
+			"GameplayTags",
+			"GameplayTagsEditor",
+			"Kismet",
+			"ToolsetRegistry",
+			"UnrealEd",
+		});
+	}
+}
diff --git a/Engine/Plugins/Experimental/Toolsets/GameplayTagsToolset/Source/GameplayTagsToolset/Private/GameplayTagsToolset/GameplayTagsToolset.cpp b/Engine/Plugins/Experimental/Toolsets/GameplayTagsToolset/Source/GameplayTagsToolset/Private/GameplayTagsToolset/GameplayTagsToolset.cpp
new file mode 100644
--- /dev/null
+++ b/Engine/Plugins/Experimental/Toolsets/GameplayTagsToolset/Source/GameplayTagsToolset/Private/GameplayTagsToolset/GameplayTagsToolset.cpp
@@ -0,0 +1,146 @@
+// Copyright Epic Games, Inc. All Rights Reserved.
+
+#include "GameplayTagsToolset/GameplayTagsToolset.h"
+
+#include "GameplayTagsEditorModule.h"
+#include "GameplayTagsManager.h"
+#include "Kismet/KismetSystemLibrary.h"
+#include "Modules/ModuleManager.h"
+
+#include UE_INLINE_GENERATED_CPP_BY_NAME(GameplayTagsToolset)
+
+namespace
+{
+	TArray<FString> CollectDescendantTagNames(const TSharedPtr<FGameplayTagNode>& ParentNode)
+	{
+		TArray<FString> TagNames;
+		TArray<TSharedPtr<FGameplayTagNode>> NodesToTraverse = ParentNode->GetChildTagNodes();
+		while (NodesToTraverse.Num() > 0)
+		{
+			TSharedPtr<FGameplayTagNode> Node = NodesToTraverse.Pop(EAllowShrinking::No);
+			if (Node.IsValid())
+			{
+				TagNames.Add(Node->GetCompleteTagName().ToString());
+				NodesToTraverse.Append(Node->GetChildTagNodes());
+			}
+		}
+		return TagNames;
+	}
+
+	IGameplayTagsEditorModule& GetModule()
+	{
+		return FModuleManager::LoadModuleChecked<IGameplayTagsEditorModule>(
+			TEXT("GameplayTagsEditor"));
+	}
+}
+
+TArray<FString> UGameplayTagsToolset::ListTags(const FString& ParentTag)
+{
+	UGameplayTagsManager& TagManager = UGameplayTagsManager::Get();
+	TArray<FString> TagNames;
+
+	if (ParentTag.IsEmpty())
+	{
+		FGameplayTagContainer AllTags;
+		TagManager.RequestAllGameplayTags(AllTags, false);
+		for (const FGameplayTag& Tag : AllTags)
+		{
+			TagNames.Add(Tag.GetTagName().ToString());
+		}
+	}
+	else
+	{
+		TSharedPtr<FGameplayTagNode> ParentNode = TagManager.FindTagNode(FName(*ParentTag));
+		if (ParentNode.IsValid())
+		{
+			TagNames = CollectDescendantTagNames(ParentNode);
+		}
+	}
+
+	TagNames.Sort();
+	return TagNames;
+}
+
+FGameplayTagInfo UGameplayTagsToolset::GetTagInfo(const FString& TagName)
+{
+	TSharedPtr<FGameplayTagNode> TagNode =
+		UGameplayTagsManager::Get().FindTagNode(FName(*TagName));
+
+	if (!TagNode.IsValid())
+	{
+		UKismetSystemLibrary::RaiseScriptError(
+			FString::Printf(TEXT("Gameplay tag '%s' does not exist."), *TagName));
+		return FGameplayTagInfo();
+	}
+
+	FGameplayTagInfo Info;
+#if WITH_EDITORONLY_DATA
+	Info.Comment = TagNode->GetDevComment();
+	Info.Source = TagNode->GetFirstSourceName().ToString();
+#endif
+
+	for (const TSharedPtr<FGameplayTagNode>& Child : TagNode->GetChildTagNodes())
+	{
+		if (Child.IsValid())
+		{
+			Info.Children.Add(Child->GetCompleteTagName().ToString());
+		}
+	}
+	return Info;
+}
+
+void UGameplayTagsToolset::AddTag(
+	const FString& TagName, const FString& Comment, const FString& TagSource)
+{
+	if (TagName.IsEmpty())
+	{
+		UKismetSystemLibrary::RaiseScriptError(TEXT("Tag name cannot be empty."));
+		return;
+	}
+
+	const FName SourceName = TagSource.IsEmpty() ? NAME_None : FName(*TagSource);
+	if (!GetModule().AddNewGameplayTagToINI(TagName, Comment, SourceName))
+	{
+		UKismetSystemLibrary::RaiseScriptError(
+			FString::Printf(TEXT("Failed to add gameplay tag '%s'."), *TagName));
+	}
+}
+
+void UGameplayTagsToolset::RemoveTag(const FString& TagName)
+{
+	TSharedPtr<FGameplayTagNode> TagNode =
+		UGameplayTagsManager::Get().FindTagNode(FName(*TagName));
+	if (!TagNode.IsValid())
+	{
+		UKismetSystemLibrary::RaiseScriptError(
+			FString::Printf(TEXT("Gameplay tag '%s' does not exist."), *TagName));
+		return;
+	}
+	if (!GetModule().DeleteTagFromINI(TagNode))
+	{
+		UKismetSystemLibrary::RaiseScriptError(
+			FString::Printf(TEXT("Failed to remove gameplay tag '%s'."), *TagName));
+	}
+}
+
+void UGameplayTagsToolset::RenameTag(const FString& OldTagName, const FString& NewTagName)
+{
+	if (OldTagName.IsEmpty() || NewTagName.IsEmpty())
+	{
+		UKismetSystemLibrary::RaiseScriptError(TEXT("Tag names cannot be empty."));
+		return;
+	}
+
+	if (!UGameplayTagsManager::Get().FindTagNode(FName(*OldTagName)).IsValid())
+	{
+		UKismetSystemLibrary::RaiseScriptError(
+			FString::Printf(TEXT("Gameplay tag '%s' does not exist."), *OldTagName));
+		return;
+	}
+
+	if (!GetModule().RenameTagInINI(OldTagName, NewTagName))
+	{
+		UKismetSystemLibrary::RaiseScriptError(FString::Printf(
+			TEXT("Failed to rename gameplay tag '%s' to '%s'."), *OldTagName, *NewTagName));
+	}
+}
diff --git a/Engine/Plugins/Experimental/Toolsets/GameplayTagsToolset/Source/GameplayTagsToolset/Private/GameplayTagsToolset/GameplayTagsToolset.h b/Engine/Plugins/Experimental/Toolsets/GameplayTagsToolset/Source/GameplayTagsToolset/Private/GameplayTagsToolset/GameplayTagsToolset.h
new file mode 100644
--- /dev/null
+++ b/Engine/Plugins/Experimental/Toolsets/GameplayTagsToolset/Source/GameplayTagsToolset/Private/GameplayTagsToolset/GameplayTagsToolset.h
@@ -0,0 +1,82 @@
+// Copyright Epic Games, Inc. All Rights Reserved.
+
+#pragma once
+
+#include "CoreMinimal.h"
+
+#include "ToolsetRegistry/UToolsetRegistry.h"
+
+#include "GameplayTagsToolset.generated.h"
+
+/// Describes the properties of a gameplay tag.
+USTRUCT(BlueprintType)
+struct FGameplayTagInfo
+{
+	GENERATED_BODY()
+
+	/// The developer comment describing the tag's purpose.
+	UPROPERTY(BlueprintReadWrite, Category = "GameplayTag")
+	FString Comment;
+
+	/// The name of the INI source file that defines this tag (e.g. "DefaultGameplayTags.ini").
+	UPROPERTY(BlueprintReadWrite, Category = "GameplayTag")
+	FString Source;
+
+	/// The fully-qualified names of this tag's immediate children.
+	UPROPERTY(BlueprintReadWrite, Category = "GameplayTag")
+	TArray<FString> Children;
+};
+
+/// Provides tools for reading and managing gameplay tags in the project's GameplayTagsManager.
+UCLASS(BlueprintType, Hidden)
+class UGameplayTagsToolset : public UToolsetDefinition
+{
+	GENERATED_BODY()
+
+public:
+	/**
+	 * Returns gameplay tags registered in the project.
+	 * @param ParentTag If non-empty, only tags that are descendants of this tag are returned.
+	 *   For example, passing "Character.State" returns "Character.State.Dead",
+	 *   "Character.State.Stunned", etc. Pass an empty string to return all tags.
+	 * @return A sorted list of fully-qualified tag names.
+	 */
+	UFUNCTION(BlueprintCallable, Category = "GameplayTags")
+	static TArray<FString> ListTags(const FString& ParentTag);
+
+	/**
+	 * Returns detailed information about a specific gameplay tag.
+	 * @param TagName The fully-qualified name of the tag, e.g. "Character.State.Dead".
+	 * @return The tag's developer comment, source INI file, and immediate children.
+	 *   Raises a script error if the tag does not exist.
+	 */
+	UFUNCTION(BlueprintCallable, Category = "GameplayTags")
+	static FGameplayTagInfo GetTagInfo(const FString& TagName);
+
+	/**
+	 * Adds a new gameplay tag to the project.
+	 * This should ONLY be called after getting explicit direction or permission from the user.
+	 * @param TagName The fully-qualified name of the tag to add, e.g. "Character.State.Dead".
+	 * @param Comment An optional developer comment describing the tag's purpose.
+	 * @param TagSource The INI source to add the tag to. Uses the default source if empty.
+	 */
+	UFUNCTION(BlueprintCallable, Category = "GameplayTags")
+	static void AddTag(const FString& TagName, const FString& Comment, const FString& TagSource);
+
+	/**
+	 * Removes a gameplay tag from the project.
+	 * This should ONLY be called after getting explicit direction or permission from the user.
+	 * @param TagName The fully-qualified name of the tag to remove, e.g. "Character.State.Dead".
+	 */
+	UFUNCTION(BlueprintCallable, Category = "GameplayTags")
+	static void RemoveTag(const FString& TagName);
+
+	/**
+	 * Renames a gameplay tag, updating all references in the project.
+	 * This should ONLY be called after getting explicit direction or permission from the user.
+	 * @param OldTagName The fully-qualified name of the tag to rename.
+	 * @param NewTagName The new fully-qualified name for the tag.
+	 */
+	UFUNCTION(BlueprintCallable, Category = "GameplayTags")
+	static void RenameTag(const FString& OldTagName, const FString& NewTagName);
+};
diff --git a/Engine/Plugins/Experimental/Toolsets/GameplayTagsToolset/Source/GameplayTagsToolset/Private/GameplayTagsToolset/Module.cpp b/Engine/Plugins/Experimental/Toolsets/GameplayTagsToolset/Source/GameplayTagsToolset/Private/GameplayTagsToolset/Module.cpp
new file mode 100644
--- /dev/null
+++ b/Engine/Plugins/Experimental/Toolsets/GameplayTagsToolset/Source/GameplayTagsToolset/Private/GameplayTagsToolset/Module.cpp
@@ -0,0 +1,25 @@
+// Copyright Epic Games, Inc. All Rights Reserved.
+
+#include "GameplayTagsToolset.h"
+
+#include "Editor.h"
+#include "Misc/CoreDelegates.h"
+#include "Modules/ModuleInterface.h"
+#include "Modules/ModuleManager.h"
+#include "ToolsetRegistry/UToolsetRegistry.h"
+
+class FGameplayTagsToolsetModule : public IModuleInterface
+{
+public:
+	virtual void StartupModule() override
+	{
+		UToolsetRegistry::RegisterToolsetClass(UGameplayTagsToolset::StaticClass());
+	}
+
+	virtual void ShutdownModule() override
+	{
+		UToolsetRegistry::UnregisterToolsetClass(UGameplayTagsToolset::StaticClass());
+	}
+};
+
+IMPLEMENT_MODULE(FGameplayTagsToolsetModule, GameplayTagsToolset);
diff --git a/Engine/Plugins/Experimental/Toolsets/GameplayTagsToolset/Source/GameplayTagsToolset/Private/GameplayTagsToolset/Tests/GameplayTagToolsetTest.cpp b/Engine/Plugins/Experimental/Toolsets/GameplayTagsToolset/Source/GameplayTagsToolset/Private/GameplayTagsToolset/Tests/GameplayTagToolsetTest.cpp
new file mode 100644
--- /dev/null
+++ b/Engine/Plugins/Experimental/Toolsets/GameplayTagsToolset/Source/GameplayTagsToolset/Private/GameplayTagsToolset/Tests/GameplayTagToolsetTest.cpp
@@ -0,0 +1,283 @@
+// Copyright Epic Games, Inc. All Rights Reserved.
+
+#include "CoreMinimal.h"
+#include "Containers/Array.h"
+#include "GameplayTagsManager.h"
+#include "Misc/AutomationTest.h"
+#include "Misc/Guid.h"
+#include "Templates/UniquePtr.h"
+#include "ToolsetRegistry/ToolCallExceptionHandler.h"
+
+#include "GameplayTagsToolset/GameplayTagsToolset.h"
+
+#if WITH_DEV_AUTOMATION_TESTS
+
+namespace
+{
+	const auto TestFlags =
+		EAutomationTestFlags::EditorContext |
+		EAutomationTestFlags::ProductFilter |
+		EAutomationTestFlags::CriticalPriority;
+}
+
+BEGIN_DEFINE_SPEC(FGameplayTagsToolsetSpec, "AI.Toolsets.GameplayTagsToolset.GameplayTagsToolsetSpec",
+	TestFlags)
+
+	// Ensure no blueprint script exceptions were caught and clear the handler.
+	void ExpectNoException()
+	{
+		check(ExceptionHandler.IsValid());
+		TestEqual(TEXT("Error"), *ExceptionHandler->GetException(), TEXT(""));
+		ExceptionHandler.Reset();
+	}
+
+	// Ensure blueprint script exception was caught and contains the error message, then clear the
+	// handler.
+	void ExpectExceptionContains(const FString& ErrorMessage)
+	{
+		check(ExceptionHandler.IsValid());
+		TestTrue(
+			*FString::Printf(TEXT("Error contains '%s'"), *ErrorMessage),
+			ExceptionHandler->GetException().Contains(ErrorMessage));
+		ExceptionHandler.Reset();
+	}
+
+	// Ensure blueprint script exception was caught and clear the handler.
+	void ExpectException()
+	{
+		check(ExceptionHandler.IsValid());
+		TestNotEqual(TEXT("Has Error"), *ExceptionHandler->GetException(), TEXT(""));
+		ExceptionHandler.Reset();
+	}
+
+	// Add a test tag that is cleaned up after the test case.
+	// Returns the name of the new tag.
+	FString AddTestTag(
+		const FString& TagNamePostfix,
+		const FString& Comment = DefaultTagComment)
+	{
+		check(ExceptionHandler);
+		FString TagName = TestTagName + TagNamePostfix;
+		ExceptionHandler->CaptureErrorsIn(
+			[this, &TagName, &Comment]()
+			{
+				UGameplayTagsToolset::AddTag(TagName, Comment, FString() /* Default source */);
+			});
+		AddedTagNames.Add(TagName);
+		return TagName;
+	}
+
+private:
+	FString TestTagName;
+	FString RenamedTagName;
+	TArray<FString> AddedTagNames;
+	TUniquePtr<UE::ToolsetRegistry::FToolCallExceptionHandler> ExceptionHandler;
+
+private:
+	static const FString NonExistentTag;
+	static const FString DefaultTagComment;
+END_DEFINE_SPEC(FGameplayTagsToolsetSpec)
+
+const FString FGameplayTagsToolsetSpec::NonExistentTag = TEXT("ZZZ.NonExistent.XYZABC123456");
+const FString FGameplayTagsToolsetSpec::DefaultTagComment = TEXT("Automated test tag");
+
+void FGameplayTagsToolsetSpec::Define()
+{
+	using namespace UE::ToolsetRegistry;
+
+	BeforeEach([this]()
+	{
+		ExceptionHandler = MakeUnique<UE::ToolsetRegistry::FToolCallExceptionHandler>();
+		TestTagName = FString::Printf(
+			TEXT("ToolsetTest.AutoTest%s"),
+			*FGuid::NewGuid().ToString(EGuidFormats::Digits).Left(8));
+		RenamedTagName = TestTagName + TEXT("Renamed");
+		AddedTagNames.Add(RenamedTagName);  // Clean up the renamed tag if it exists.
+	});
+
+	AfterEach([this]()
+	{
+		if (ExceptionHandler.IsValid()) ExpectNoException();
+		// Clean up test tags.
+		UGameplayTagsManager& Manager = UGameplayTagsManager::Get();
+		for (const FString& TagName : AddedTagNames)
+		{
+			if (Manager.FindTagNode(FName(*TagName)).IsValid())
+			{
+				UGameplayTagsToolset::RemoveTag(TagName);
+			}
+		}
+		TestTagName.Empty();
+		RenamedTagName.Empty();
+		AddedTagNames.Empty();
+	});
+
+	Describe(TEXT("ListTags"), [this]()
+	{
+		It(TEXT("Returns a sorted list"), [this]()
+		{
+			AddTestTag(TEXT("foo"));
+			AddTestTag(TEXT("foo.bar"));
+			ExceptionHandler->CaptureErrorsIn([this]()
+				{
+					TArray<FString> Tags = UGameplayTagsToolset::ListTags(FString());
+					TArray<FString> Sorted = Tags;
+					Sorted.Sort();
+					TestEqual("Sorted", Tags, Sorted);
+				});
+		});
+
+		It(TEXT("Returns empty array for a non-existent parent tag"), [this]()
+		{
+			ExceptionHandler->CaptureErrorsIn([this]()
+				{
+					TArray<FString> Tags = UGameplayTagsToolset::ListTags(NonExistentTag);
+					TestEqual("Empty", Tags.Num(), 0);
+				});
+		});
+	});
+
+	Describe(TEXT("GetTagInfo"), [this]()
+	{
+		It(TEXT("Fails for a non-existent tag"), [this]()
+		{
+			ExceptionHandler->CaptureErrorsIn([this]()
+				{
+					FGameplayTagInfo Info = UGameplayTagsToolset::GetTagInfo(NonExistentTag);
+#if WITH_EDITORONLY_DATA
+					TestTrue("Comment empty", Info.Comment.IsEmpty());
+					TestTrue("Source empty", Info.Source.IsEmpty());
+#endif  // WITH_EDITORONLY_DATA
+					TestEqual("No children", Info.Children.Num(), 0);
+				});
+			ExpectException();
+		});
+
+		It(TEXT("Returns info for an existing tag"), [this]()
+		{
+			FString TagName = AddTestTag(TEXT(".foo"));
+			TArray<FString> Children =
+			{
+				AddTestTag(TEXT(".foo.bish")),
+				AddTestTag(TEXT(".foo.bosh")),
+			};
+			ExceptionHandler->CaptureErrorsIn([this, &TagName, &Children]()
+				{
+					FGameplayTagInfo Info = UGameplayTagsToolset::GetTagInfo(TagName);
+#if WITH_EDITORONLY_DATA
+					TestEqual(TEXT("Comment"), Info.Comment, DefaultTagComment);
+					TestEqual(
+						TEXT("Source"), Info.Source,
+						FGameplayTagSource::GetDefaultName().ToString());
+#endif  // WITH_EDITORONLY_DATA
+					if (TestEqual(TEXT("NumChildren"), Info.Children.Num(), Children.Num()))
+					{
+						for (int i = 0; i < Children.Num(); ++i)
+						{
+							TestEqual(
+								*FString::Printf(TEXT("Children[%d]"), i),
+								Info.Children[i], Children[i]);
+						}
+					}
+				});
+			ExpectNoException();
+		});
+	});
+
+	Describe(TEXT("AddTag and RemoveTag"), [this]()
+	{
+		It(TEXT("Can add a new tag"), [this]()
+		{
+			FString TagName;
+			ExceptionHandler->CaptureErrorsIn(
+				[this, &TagName]()
+				{
+					TagName = AddTestTag(TEXT("test"));
+				});
+			ExpectNoException();
+			TestTrue(TEXT("Exists in manager"),
+				UGameplayTagsManager::Get().FindTagNode(FName(*TagName)).IsValid());
+		});
+
+		It(TEXT("Fails to add empty tag"), [this]()
+		{
+			ExceptionHandler->CaptureErrorsIn(
+				[this]()
+				{
+					UGameplayTagsToolset::AddTag(FString(), FString(), FString());
+				});
+			ExpectException();
+		});
+
+		It(TEXT("Can remove an existing tag"), [this]()
+		{
+			FString TagName = AddTestTag(TEXT(""));
+			ExceptionHandler->CaptureErrorsIn(
+				[this, &TagName]()
+				{
+					UGameplayTagsToolset::RemoveTag(TagName);
+				});
+			ExpectNoException();
+			TestFalse(TEXT("No longer in manager"),
+				UGameplayTagsManager::Get().FindTagNode(FName(*TagName)).IsValid());
+		});
+
+		It(TEXT("Fails to remove a non-existent tag"), [this]()
+		{
+			ExceptionHandler->CaptureErrorsIn(
+				[this]()
+				{
+					UGameplayTagsToolset::RemoveTag(NonExistentTag);
+				});
+			ExpectExceptionContains(NonExistentTag);
+		});
+	});
+
+	Describe(TEXT("RenameTag"), [this]()
+	{
+		It(TEXT("Can rename an existing tag"), [this]()
+		{
+			FString TagName = AddTestTag(TEXT(""));
+			ExceptionHandler->CaptureErrorsIn(
+				[this, &TagName]()
+				{
+					UGameplayTagsToolset::RenameTag(TagName, RenamedTagName);
+				});
+			ExpectNoException();
+			TestTrue(TEXT("New name exists"),
+				UGameplayTagsManager::Get().FindTagNode(FName(*RenamedTagName)).IsValid());
+		});
+
+		It(TEXT("Fails when source tag name does not exist"), [this]()
+		{
+			ExceptionHandler->CaptureErrorsIn(
+				[this]()
+				{
+					UGameplayTagsToolset::RenameTag(NonExistentTag, RenamedTagName);
+				});
+			ExpectExceptionContains(NonExistentTag);
+		});
+
+		It(TEXT("Fails with an empty source tag name"), [this]()
+		{
+			ExceptionHandler->CaptureErrorsIn(
+				[this]()
+				{
+					UGameplayTagsToolset::RenameTag(FString(), RenamedTagName);
+				});
+			ExpectException();
+		});
+
+		It(TEXT("Fails with an empty target tag name"), [this]()
+		{
+			ExceptionHandler->CaptureErrorsIn(
+				[this]()
+				{
+					UGameplayTagsToolset::RenameTag(TestTagName, FString());
+				});
+			ExpectException();
+		});
+	});
+}
+
+#endif  // WITH_DEV_AUTOMATION_TESTS
