[GameplayTagsToolset] Add FindReferencersByTag tool

Adds a new FindReferencersByTag function to UGameplayTagsToolset that queries the asset registry for assets referencing a gameplay tag via searchable names. Returns an array of package paths. Includes three automation test cases covering error handling and empty results. Adds AssetRegistry module dependency.

#Synced-CL 52555021

diff --git a/Engine/Plugins/Experimental/Toolsets/GameplayTagsToolset/Source/GameplayTagsToolset/GameplayTagsToolset.Build.cs b/Engine/Plugins/Experimental/Toolsets/GameplayTagsToolset/Source/GameplayTagsToolset/GameplayTagsToolset.Build.cs
--- a/Engine/Plugins/Experimental/Toolsets/GameplayTagsToolset/Source/GameplayTagsToolset/GameplayTagsToolset.Build.cs
+++ b/Engine/Plugins/Experimental/Toolsets/GameplayTagsToolset/Source/GameplayTagsToolset/GameplayTagsToolset.Build.cs
@@ -14,7 +14,9 @@ public class GameplayTagsToolset : ModuleRules
 
 		PrivateDependencyModuleNames.AddRange(new string[]
 		{
+			"AssetRegistry",
 			"CoreUObject",
+			"EditorScriptingUtilities",
 			"Engine",
 			"GameplayTags",
 			"GameplayTagsEditor",
diff --git a/Engine/Plugins/Experimental/Toolsets/GameplayTagsToolset/Source/GameplayTagsToolset/Private/GameplayTagsToolset/GameplayTagsToolset.cpp b/Engine/Plugins/Experimental/Toolsets/GameplayTagsToolset/Source/GameplayTagsToolset/Private/GameplayTagsToolset/GameplayTagsToolset.cpp
--- a/Engine/Plugins/Experimental/Toolsets/GameplayTagsToolset/Source/GameplayTagsToolset/Private/GameplayTagsToolset/GameplayTagsToolset.cpp
+++ b/Engine/Plugins/Experimental/Toolsets/GameplayTagsToolset/Source/GameplayTagsToolset/Private/GameplayTagsToolset/GameplayTagsToolset.cpp
@@ -5,6 +5,7 @@
 #include "GameplayTagsEditorModule.h"
 #include "GameplayTagsManager.h"
 #include "Kismet/KismetSystemLibrary.h"
+#include "AssetRegistry/IAssetRegistry.h"
 #include "Modules/ModuleManager.h"
 
 #include UE_INLINE_GENERATED_CPP_BY_NAME(GameplayTagsToolset)
@@ -144,3 +145,34 @@ void UGameplayTagsToolset::RenameTag(const FString& OldTagName, const FString& N
 			TEXT("Failed to rename gameplay tag '%s' to '%s'."), *OldTagName, *NewTagName));
 	}
 }
+
+TArray<FString> UGameplayTagsToolset::FindReferencersByTag(const FString& TagName)
+{
+	if (TagName.IsEmpty())
+	{
+		UKismetSystemLibrary::RaiseScriptError(TEXT("Tag name cannot be empty."));
+		return TArray<FString>();
+	}
+
+	const FGameplayTag Tag = UGameplayTagsManager::Get().RequestGameplayTag(FName(*TagName), false);
+	if (!Tag.IsValid())
+	{
+		UKismetSystemLibrary::RaiseScriptError(
+			FString::Printf(TEXT("Gameplay tag '%s' does not exist."), *TagName));
+		return TArray<FString>();
+	}
+
+	const FAssetIdentifier TagIdentifier(FGameplayTag::StaticStruct(), Tag.GetTagName());
+
+	TArray<FAssetIdentifier> Referencers;
+	IAssetRegistry::Get()->GetReferencers(TagIdentifier, Referencers,
+		UE::AssetRegistry::EDependencyCategory::SearchableName);
+
+	TArray<FString> Result;
+	Result.Reserve(Referencers.Num());
+	for (const FAssetIdentifier& Id : Referencers)
+	{
+		Result.Add(Id.PackageName.ToString());
+	}
+	return Result;
+}
diff --git a/Engine/Plugins/Experimental/Toolsets/GameplayTagsToolset/Source/GameplayTagsToolset/Private/GameplayTagsToolset/GameplayTagsToolset.h b/Engine/Plugins/Experimental/Toolsets/GameplayTagsToolset/Source/GameplayTagsToolset/Private/GameplayTagsToolset/GameplayTagsToolset.h
--- a/Engine/Plugins/Experimental/Toolsets/GameplayTagsToolset/Source/GameplayTagsToolset/Private/GameplayTagsToolset/GameplayTagsToolset.h
+++ b/Engine/Plugins/Experimental/Toolsets/GameplayTagsToolset/Source/GameplayTagsToolset/Private/GameplayTagsToolset/GameplayTagsToolset.h
@@ -79,4 +79,12 @@ public:
 	 */
 	UFUNCTION(BlueprintCallable, Category = "GameplayTags")
 	static void RenameTag(const FString& OldTagName, const FString& NewTagName);
+
+	/**
+	 * Returns assets that reference a gameplay tag.
+	 * @param TagName The fully-qualified gameplay tag to search for, e.g. "Character.State.Dead".
+	 * @return A list of package paths that reference the tag.
+	 */
+	UFUNCTION(BlueprintCallable, Category = "GameplayTags")
+	static TArray<FString> FindReferencersByTag(const FString& TagName);
 };
diff --git a/Engine/Plugins/Experimental/Toolsets/GameplayTagsToolset/Source/GameplayTagsToolset/Private/GameplayTagsToolset/Tests/GameplayTagToolsetTest.cpp b/Engine/Plugins/Experimental/Toolsets/GameplayTagsToolset/Source/GameplayTagsToolset/Private/GameplayTagsToolset/Tests/GameplayTagToolsetTest.cpp
--- a/Engine/Plugins/Experimental/Toolsets/GameplayTagsToolset/Source/GameplayTagsToolset/Private/GameplayTagsToolset/Tests/GameplayTagToolsetTest.cpp
+++ b/Engine/Plugins/Experimental/Toolsets/GameplayTagsToolset/Source/GameplayTagsToolset/Private/GameplayTagsToolset/Tests/GameplayTagToolsetTest.cpp
@@ -1,10 +1,17 @@
 // Copyright Epic Games, Inc. All Rights Reserved.
 
-#include "CoreMinimal.h"
+#include "GameplayTagsToolset/Tests/GameplayTagToolsetTest.h"
+
+#include "AssetRegistry/AssetRegistryModule.h"
+#include "AssetToolsModule.h"
 #include "Containers/Array.h"
+#include "CoreMinimal.h"
+#include "EditorAssetLibrary.h"
 #include "GameplayTagsManager.h"
 #include "Misc/AutomationTest.h"
 #include "Misc/Guid.h"
+#include "Misc/PackageName.h"
+#include "Misc/Paths.h"
 #include "Templates/UniquePtr.h"
 #include "ToolsetRegistry/ToolCallExceptionHandler.h"
 
@@ -278,6 +285,101 @@ void FGameplayTagsToolsetSpec::Define()
 			ExpectException();
 		});
 	});
+
+	Describe(TEXT("FindReferencersByTag"), [this]()
+	{
+		It(TEXT("Returns empty array for a tag with no referencers"), [this]()
+		{
+			FString TagName = AddTestTag(TEXT(".norefs"));
+			ExceptionHandler->CaptureErrorsIn([this, &TagName]()
+				{
+					TArray<FString> Referencers =
+						UGameplayTagsToolset::FindReferencersByTag(TagName);
+					TestEqual("Empty", Referencers.Num(), 0);
+				});
+			ExpectNoException();
+		});
+
+		It(TEXT("Fails for a non-existent tag"), [this]()
+		{
+			ExceptionHandler->CaptureErrorsIn([this]()
+				{
+					TArray<FString> Referencers =
+						UGameplayTagsToolset::FindReferencersByTag(NonExistentTag);
+					TestEqual("Empty", Referencers.Num(), 0);
+				});
+			ExpectExceptionContains(NonExistentTag);
+		});
+
+		It(TEXT("Fails with an empty tag name"), [this]()
+		{
+			ExceptionHandler->CaptureErrorsIn([this]()
+				{
+					UGameplayTagsToolset::FindReferencersByTag(FString());
+				});
+			ExpectException();
+		});
+
+		It(TEXT("Returns referencers for a tag used by a saved asset"), [this]()
+		{
+			FString TagName = AddTestTag(TEXT(".hasrefs"));
+
+			// Mount a transient content path for test assets.
+			const FString MountPoint = TEXT("/Automation/");
+			FPackageName::RegisterMountPoint(
+				*MountPoint, FPaths::AutomationTransientDir());
+
+			// Create a minimal asset that contains our test tag via
+			// the editor asset pipeline (avoids data validation issues
+			// with raw UPackage::SavePackage).
+			const FString AssetName = TEXT("TagRefTestAsset");
+			const FString PackagePath = MountPoint + AssetName;
+			IAssetTools& AssetTools =
+				FAssetToolsModule::GetModule().Get();
+			UObject* CreatedAsset = AssetTools.CreateAsset(
+				AssetName, MountPoint,
+				UGameplayTagTestAsset::StaticClass(), nullptr);
+			if (!TestNotNull(TEXT("Asset created"), CreatedAsset))
+			{
+				return;
+			}
+
+			UGameplayTagTestAsset* TestAsset =
+				CastChecked<UGameplayTagTestAsset>(CreatedAsset);
+			TestAsset->Tags.AddTag(
+				UGameplayTagsManager::Get().RequestGameplayTag(
+					FName(*TagName)));
+
+			// Save through the editor library so the asset registry
+			// picks up the searchable name.
+			UEditorAssetLibrary::SaveAsset(PackagePath);
+
+			// Ensure the asset registry has processed the saved file's
+			// searchable name dependencies before we query.
+			const FString FilePath =
+				FPaths::AutomationTransientDir() / TEXT("TagRefTestAsset.uasset");
+			IAssetRegistry& AssetRegistry =
+				FModuleManager::LoadModuleChecked<FAssetRegistryModule>(
+					TEXT("AssetRegistry")).Get();
+			AssetRegistry.ScanFilesSynchronous({FilePath}, true /* bForceRescan */);
+
+			// Query and verify.
+			ExceptionHandler->CaptureErrorsIn([this, &TagName, &PackagePath]()
+				{
+					TArray<FString> Referencers =
+						UGameplayTagsToolset::FindReferencersByTag(TagName);
+					TestTrue("Has referencers", Referencers.Num() > 0);
+					TestTrue("Contains test asset",
+						Referencers.Contains(PackagePath));
+				});
+			ExpectNoException();
+
+			// Cleanup.
+			UEditorAssetLibrary::DeleteAsset(PackagePath);
+			FPackageName::UnRegisterMountPoint(*MountPoint,
+				FPaths::AutomationTransientDir());
+		});
+	});
 }
 
 #endif  // WITH_DEV_AUTOMATION_TESTS
diff --git a/Engine/Plugins/Experimental/Toolsets/GameplayTagsToolset/Source/GameplayTagsToolset/Private/GameplayTagsToolset/Tests/GameplayTagToolsetTest.h b/Engine/Plugins/Experimental/Toolsets/GameplayTagsToolset/Source/GameplayTagsToolset/Private/GameplayTagsToolset/Tests/GameplayTagToolsetTest.h
new file mode 100644
--- /dev/null
+++ b/Engine/Plugins/Experimental/Toolsets/GameplayTagsToolset/Source/GameplayTagsToolset/Private/GameplayTagsToolset/Tests/GameplayTagToolsetTest.h
@@ -0,0 +1,20 @@
+// Copyright Epic Games, Inc. All Rights Reserved.
+
+#pragma once
+
+#include "CoreMinimal.h"
+
+#include "Engine/DataAsset.h"
+#include "GameplayTagContainer.h"
+#include "GameplayTagToolsetTest.generated.h"
+
+// Minimal test asset that holds a gameplay tag container for searchable name tests.
+UCLASS(Hidden)
+class UGameplayTagTestAsset : public UPrimaryDataAsset
+{
+	GENERATED_BODY()
+
+public:
+	UPROPERTY()
+	FGameplayTagContainer Tags;
+};
