[AI Assistant] Fix GameplayTagsToolsetSpec tests failing on build machines

BeforeEach now pre-creates an empty INI in FPaths::ProjectSavedDir() / "AutomationTests" / "GameplayTags", registers it with AddTagIniSearchPath(), and rebuilds the tag tree via EditorRefreshGameplayTagTree(). This causes the tag manager to set the source's ConfigFileName to the writable saved-dir path, so that AddNewGameplayTagToINI writes there rather than to the read-only project config directory that build machines mount as read-only.

#Synced-CL 52584985

diff --git a/Engine/Plugins/Experimental/Toolsets/GameplayTagsToolset/Source/GameplayTagsToolset/Private/GameplayTagsToolset/Tests/GameplayTagToolsetTest.cpp b/Engine/Plugins/Experimental/Toolsets/GameplayTagsToolset/Source/GameplayTagsToolset/Private/GameplayTagsToolset/Tests/GameplayTagToolsetTest.cpp
--- a/Engine/Plugins/Experimental/Toolsets/GameplayTagsToolset/Source/GameplayTagsToolset/Private/GameplayTagsToolset/Tests/GameplayTagToolsetTest.cpp
+++ b/Engine/Plugins/Experimental/Toolsets/GameplayTagsToolset/Source/GameplayTagsToolset/Private/GameplayTagsToolset/Tests/GameplayTagToolsetTest.cpp
@@ -8,7 +8,9 @@
 #include "CoreMinimal.h"
 #include "EditorAssetLibrary.h"
 #include "GameplayTagsManager.h"
+#include "HAL/FileManager.h"
 #include "Misc/AutomationTest.h"
+#include "Misc/FileHelper.h"
 #include "Misc/Guid.h"
 #include "Misc/PackageName.h"
 #include "Misc/Paths.h"
@@ -68,7 +70,7 @@ BEGIN_DEFINE_SPEC(FGameplayTagsToolsetSpec, "AI.Toolsets.GameplayTagsToolset.Gam
 		ExceptionHandler->CaptureErrorsIn(
 			[this, &TagName, &Comment]()
 			{
-				UGameplayTagsToolset::AddTag(TagName, Comment, FString() /* Default source */);
+				UGameplayTagsToolset::AddTag(TagName, Comment, TestTagSourceName.ToString());
 			});
 		AddedTagNames.Add(TagName);
 		return TagName;
@@ -83,10 +85,15 @@ private:
 private:
 	static const FString NonExistentTag;
 	static const FString DefaultTagComment;
+	// Tag source backed by the project saved directory so writes succeed on build machines,
+	// where the project config directory is read-only. AddTagIniSearchPath registers it with
+	// the full writable path so AddNewGameplayTagToINI can find and use it.
+	static const FName TestTagSourceName;
 END_DEFINE_SPEC(FGameplayTagsToolsetSpec)
 
 const FString FGameplayTagsToolsetSpec::NonExistentTag = TEXT("ZZZ.NonExistent.XYZABC123456");
 const FString FGameplayTagsToolsetSpec::DefaultTagComment = TEXT("Automated test tag");
+const FName FGameplayTagsToolsetSpec::TestTagSourceName = FName(TEXT("GameplayTagsToolsetTest.ini"));
 
 void FGameplayTagsToolsetSpec::Define()
 {
@@ -94,6 +101,28 @@ void FGameplayTagsToolsetSpec::Define()
 
 	BeforeEach([this]()
 	{
+		// Register a tag INI source in the project saved directory so that
+		// AddNewGameplayTagToINI writes to a location that is writable on all machines,
+		// including build machines where the project config directory is read-only.
+		//
+		// AddTagIniSearchPath discovers .ini files in the directory and records their full
+		// paths. EditorRefreshGameplayTagTree then registers each file as a tag source with
+		// its correct ConfigFileName (pointing into the saved directory), so that subsequent
+		// AddTag calls write to the right writable location.
+		const FString TestTagIniDir =
+			FPaths::ProjectSavedDir() / TEXT("AutomationTests") / TEXT("GameplayTags");
+		const FString TestTagIniPath = TestTagIniDir / TestTagSourceName.ToString();
+
+		IFileManager::Get().MakeDirectory(*TestTagIniDir, /*Tree=*/true);
+		if (!FPaths::FileExists(TestTagIniPath))
+		{
+			FFileHelper::SaveStringToFile(TEXT(""), *TestTagIniPath);
+		}
+
+		UGameplayTagsManager& Manager = UGameplayTagsManager::Get();
+		Manager.AddTagIniSearchPath(TestTagIniDir);
+		Manager.EditorRefreshGameplayTagTree();
+
 		ExceptionHandler = MakeUnique<UE::ToolsetRegistry::FToolCallExceptionHandler>();
 		TestTagName = FString::Printf(
 			TEXT("ToolsetTest.AutoTest%s"),
@@ -173,9 +202,7 @@ void FGameplayTagsToolsetSpec::Define()
 					FGameplayTagInfo Info = UGameplayTagsToolset::GetTagInfo(TagName);
 #if WITH_EDITORONLY_DATA
 					TestEqual(TEXT("Comment"), Info.Comment, DefaultTagComment);
-					TestEqual(
-						TEXT("Source"), Info.Source,
-						FGameplayTagSource::GetDefaultName().ToString());
+					TestEqual(TEXT("Source"), Info.Source, TestTagSourceName.ToString());
 #endif  // WITH_EDITORONLY_DATA
 					if (TestEqual(TEXT("NumChildren"), Info.Children.Num(), Children.Num()))
 					{
